Enabling CAD importer by default for AEC and MFG templates

diff --git a/Templates/TP_AEC_BlankBP/TP_AEC_BlankBP.uproject b/Templates/TP_AEC_BlankBP/TP_AEC_BlankBP.uproject
--- a/Templates/TP_AEC_BlankBP/TP_AEC_BlankBP.uproject
+++ b/Templates/TP_AEC_BlankBP/TP_AEC_BlankBP.uproject
@@ -17,6 +17,10 @@
 			"Name": "DatasmithImporter",
 			"Enabled": true
 		},
+		{
+			"Name": "DatasmithCADImporter",
+			"Enabled": true
+		},
 		{
 			"Name": "VariantManager",
 			"Enabled": True
diff --git a/Templates/TP_MFG_BlankBP/TP_MFG_BlankBP.uproject b/Templates/TP_MFG_BlankBP/TP_MFG_BlankBP.uproject
--- a/Templates/TP_MFG_BlankBP/TP_MFG_BlankBP.uproject
+++ b/Templates/TP_MFG_BlankBP/TP_MFG_BlankBP.uproject
@@ -17,6 +17,10 @@
 			"Name": "DatasmithImporter",
 			"Enabled": true
 		},
+		{
+			"Name": "DatasmithCADImporter",
+			"Enabled": true
+		},
 		{
 			"Name": "VariantManager",
 			"Enabled": True
diff --git a/Templates/TP_PhotoStudioBP/TP_PhotoStudioBP.uproject b/Templates/TP_PhotoStudioBP/TP_PhotoStudioBP.uproject
--- a/Templates/TP_PhotoStudioBP/TP_PhotoStudioBP.uproject
+++ b/Templates/TP_PhotoStudioBP/TP_PhotoStudioBP.uproject
@@ -17,6 +17,10 @@
 			"Name": "DatasmithImporter",
 			"Enabled": true
 		},
+		{
+			"Name": "DatasmithCADImporter",
+			"Enabled": true
+		},
 		{
 			"Name": "VariantManager",
 			"Enabled": True
