Enabling polygon editing plugin for AEC end MFG templates to provide same out of box experience than with Unreal Studio

diff --git a/Templates/TP_AEC_BlankBP/TP_AEC_BlankBP.uproject b/Templates/TP_AEC_BlankBP/TP_AEC_BlankBP.uproject
--- a/Templates/TP_AEC_BlankBP/TP_AEC_BlankBP.uproject
+++ b/Templates/TP_AEC_BlankBP/TP_AEC_BlankBP.uproject
@@ -36,6 +36,10 @@
 		{
 			"Name": "SequencerScripting",
 			"Enabled": false
+		},
+				{
+			"Name": "StaticMeshEditorExtension",
+			"Enabled": true
 		}
 	]
 }
\ No newline at end of file
diff --git a/Templates/TP_MFG_BlankBP/TP_MFG_BlankBP.uproject b/Templates/TP_MFG_BlankBP/TP_MFG_BlankBP.uproject
--- a/Templates/TP_MFG_BlankBP/TP_MFG_BlankBP.uproject
+++ b/Templates/TP_MFG_BlankBP/TP_MFG_BlankBP.uproject
@@ -36,6 +36,10 @@
 		{
 			"Name": "SequencerScripting",
 			"Enabled": true
+		},
+				{
+			"Name": "StaticMeshEditorExtension",
+			"Enabled": true
 		}
 	]
 }
\ No newline at end of file
diff --git a/Templates/TP_PhotoStudioBP/TP_PhotoStudioBP.uproject b/Templates/TP_PhotoStudioBP/TP_PhotoStudioBP.uproject
--- a/Templates/TP_PhotoStudioBP/TP_PhotoStudioBP.uproject
+++ b/Templates/TP_PhotoStudioBP/TP_PhotoStudioBP.uproject
@@ -36,6 +36,10 @@
 		{
 			"Name": "SequencerScripting",
 			"Enabled": true
+		},
+				{
+			"Name": "StaticMeshEditorExtension",
+			"Enabled": true
 		}
 	]
 }
\ No newline at end of file
