Enabled Mesh Modelling plugin in the PhotoStudioBP Template project #rb none #jira UE-89867

diff --git a/Templates/TP_PhotoStudioBP/TP_PhotoStudioBP.uproject b/Templates/TP_PhotoStudioBP/TP_PhotoStudioBP.uproject
--- a/Templates/TP_PhotoStudioBP/TP_PhotoStudioBP.uproject
+++ b/Templates/TP_PhotoStudioBP/TP_PhotoStudioBP.uproject
@@ -61,6 +61,10 @@
 			"SupportedTargetPlatforms": [
 				"Win64"
 			]
+		},
+		{
+			"Name": "ModelingToolsEditorMode",
+			"Enabled": true
 		}
 	]
 }
\ No newline at end of file
