[Template-PhotoStudio/MFG_Blank] Enable the Interchange AxF plugin and disable the AxF Importer plugin

#Synced-CL 51566404

diff --git a/Templates/TP_MFG_BlankBP/TP_MFG_BlankBP.uproject b/Templates/TP_MFG_BlankBP/TP_MFG_BlankBP.uproject
--- a/Templates/TP_MFG_BlankBP/TP_MFG_BlankBP.uproject
+++ b/Templates/TP_MFG_BlankBP/TP_MFG_BlankBP.uproject
@@ -32,13 +32,6 @@
 				"Mac"
 			]
 		},
-		{
-			"Name": "AxFImporter",
-			"Enabled": true,
-			"SupportedTargetPlatforms": [
-				"Win64"
-			]
-		},
 		{
 			"Name": "DataPrepEditor",
 			"Enabled": true
@@ -74,6 +67,15 @@
 		{
 			"Name": "MoviePipelineMaskRenderPass",
 			"Enabled": true
+		},
+		{
+			"Name": "InterchangeAxF",
+			"Enabled": true,
+			"SupportedTargetPlatforms": [
+				"Win64",
+				"Linux",
+				"Mac"
+			]
 		}
 	]
 }
\ No newline at end of file
diff --git a/Templates/TP_PhotoStudioBP/TP_PhotoStudioBP.uproject b/Templates/TP_PhotoStudioBP/TP_PhotoStudioBP.uproject
--- a/Templates/TP_PhotoStudioBP/TP_PhotoStudioBP.uproject
+++ b/Templates/TP_PhotoStudioBP/TP_PhotoStudioBP.uproject
@@ -49,13 +49,6 @@
 				"Mac"
 			]
 		},
-		{
-			"Name": "AxFImporter",
-			"Enabled": true,
-			"SupportedTargetPlatforms": [
-				"Win64"
-			]
-		},
 		{
 			"Name": "ModelingToolsEditorMode",
 			"Enabled": true
@@ -63,6 +56,15 @@
 		{
 			"Name": "MovieRenderPipeline",
 			"Enabled": true
+		},
+		{
+			"Name": "InterchangeAxF",
+			"Enabled": true,
+			"SupportedTargetPlatforms": [
+				"Win64",
+				"Linux",
+				"Mac"
+			]
 		}
 	]
 }
\ No newline at end of file
